MediaCapture capture the backbuffer when the Viewport don't use an internal texture.

diff --git a/Engine/Source/Runtime/MediaIOCore/Private/MediaCapture.cpp b/Engine/Source/Runtime/MediaIOCore/Private/MediaCapture.cpp
--- a/Engine/Source/Runtime/MediaIOCore/Private/MediaCapture.cpp
+++ b/Engine/Source/Runtime/MediaIOCore/Private/MediaCapture.cpp
@@ -329,6 +329,10 @@ void UMediaCapture::OnEndFrame_GameThread()
 			if (InSceneViewportPtr.IsValid())
 			{
 				SourceTexture = InSceneViewportPtr->GetRenderTargetTexture();
+				if (!SourceTexture.IsValid() && InSceneViewportPtr->GetViewportRHI())
+				{
+					SourceTexture = RHICmdList.GetViewportBackBuffer(InSceneViewportPtr->GetViewportRHI());
+				}
 			}
 			else if (InCapturingRenderTarget)
 			{
